[Shaders] Fix old bugs in HlslParser.

- Use more explicit types.
- Remove redundant null pointer checks on `delete` operator.
- Remove double delete of sub expressions in destructor of `FExpression`; This was introduced with CL 10075936 when the `SubExpressions` field was accidentally replaced with the `Expressions` field.
- Remove unused function `GetCurrentTokenAndAdvance()` in FHlslLexer.
- Remove `check(0)` in `GetPrecedence()` for ternary operator. No need to crash in such a getter-function even if this operator is never used in this code path.
- Use TEXT() macro instead of casting single character output to TCHAR.
- Fix potential null pointer exception when return value of `GetNonDigitSemanticPrefix()` is passed to `MakeIndexedSemantic()`.
- Fix uninitialized member `ExpressionScope` in FASTWriter.
- Honor function contract of `ForEachHlslIdentifier()` by breaking out of the loop when the callback returns false.

#rnx

diff --git a/Engine/Source/Developer/HlslParser/Private/HlslAST.cpp b/Engine/Source/Developer/HlslParser/Private/HlslAST.cpp
--- a/Engine/Source/Developer/HlslParser/Private/HlslAST.cpp
+++ b/Engine/Source/Developer/HlslParser/Private/HlslAST.cpp
@@ -18,26 +18,18 @@ namespace CrossCompiler
 			}
 			else
 			{
-				for (const auto* Dimension : ArraySize)
+				for (const FExpression* Dimension : ArraySize)
 				{
-					Writer << (TCHAR)'[';
+					Writer << TEXT('[');
 					if (Dimension)
 					{
 						Dimension->Write(Writer);
 					}
-					Writer << (TCHAR)']';
+					Writer << TEXT(']');
 				}
 			}
 		}
 
-#if 0
-		FNode::FNode()/* :
-			Prev(nullptr),
-			Next(nullptr)*/
-		{
-		}
-#endif // 0
-
 		FNode::FNode(FLinearAllocator* Allocator, const FSourceInfo& InInfo) :
 			SourceInfo(InInfo),
 			Attributes(Allocator)/*,
@@ -48,10 +40,9 @@ namespace CrossCompiler
 
 		void FASTWriter::DoIndent()
 		{
-			int32 N = Indent;
-			while (--N >= 0)
+			for (int32 TabIndex = 0; TabIndex < Indent; ++TabIndex)
 			{
-				(*this) << (TCHAR)'\t';
+				(*this) << TEXT('\t');
 			}
 		}
 
@@ -64,7 +55,7 @@ namespace CrossCompiler
 					Attr->Write(Writer);
 				}
 
-				Writer << (TCHAR)' ';
+				Writer << TEXT(' ');
 			}
 		}
 
@@ -279,15 +270,15 @@ namespace CrossCompiler
 				break;
 
 			case EOperators::TypeCast:
-				Writer << (TCHAR)'(';
+				Writer << TEXT('(');
 				TypeSpecifier->Write(Writer);
-				Writer << TEXT(")");
+				Writer << TEXT(')');
 				break;
 
 			default:
 				Writer << TEXT("*MISSING_");
 				Writer << (uint32)Operator;
-				Writer << (TCHAR)'*';
+				Writer << TEXT('*');
 				checkf(0, TEXT("Unhandled AST Operator %d!"), (uint32)Operator);
 				break;
 			}
@@ -298,7 +289,7 @@ namespace CrossCompiler
 			switch (Operator)
 			{
 			case EOperators::Conditional:
-				Writer << (TCHAR)'(';
+				Writer << TEXT('(');
 				Expressions[0]->Write(Writer);
 				Writer << TEXT(" ? ");
 				Expressions[1]->Write(Writer);
@@ -318,7 +309,7 @@ namespace CrossCompiler
 			default:
 				Writer << TEXT("*MISSING_");
 				Writer << (uint32)Operator;
-				Writer << (TCHAR)'*';
+				Writer << TEXT('*');
 				checkf(0, TEXT("Unhandled AST Operator %d!"), (int32)Operator);
 				break;
 			}
@@ -338,24 +329,9 @@ namespace CrossCompiler
 
 		FExpression::~FExpression()
 		{
-			for (auto* Expr : Expressions)
+			for (FExpression* Expr : Expressions)
 			{
-				if (Expr)
-				{
-					delete Expr;
-				}
-			}
-
-			for (int32 Index = 0; Index < Expressions.Num(); ++Index)
-			{
-				if (Expressions[Index])
-				{
-					delete Expressions[Index];
-				}
-				else
-				{
-					break;
-				}
+				delete Expr;
 			}
 		}
 
@@ -370,7 +346,7 @@ namespace CrossCompiler
 
 			if (Writer.ExpressionScope != 0 && Operator != EOperators::FieldSelection)
 			{
-				Writer << (TCHAR)'(';
+				Writer << TEXT('(');
 			}
 
 			if (Expressions.Num() != 0)
@@ -392,7 +368,7 @@ namespace CrossCompiler
 				break;
 
 			case EOperators::FieldSelection:
-				Writer << (TCHAR)'.';
+				Writer << TEXT('.');
 				Writer << Identifier;
 				break;
 
@@ -402,7 +378,7 @@ namespace CrossCompiler
 
 			if (Writer.ExpressionScope != 0 && Operator != EOperators::FieldSelection)
 			{
-				Writer << (TCHAR)')';
+				Writer << TEXT(')');
 			}
 		}
 
@@ -422,30 +398,30 @@ namespace CrossCompiler
 				}
 				else
 				{
-					Writer << (TCHAR)'(';
+					Writer << TEXT('(');
 					Expressions[0]->Write(Writer);
-					Writer << (TCHAR)')';
+					Writer << TEXT(')');
 				}
-				Writer << (TCHAR)'[';
+				Writer << TEXT('[');
 				Expressions[1]->Write(Writer);
-				Writer << (TCHAR)']';
+				Writer << TEXT(']');
 				break;
 
 			default:
 				if (Writer.ExpressionScope != 0 && !IsAssignmentOperator(Operator))
 				{
-					Writer << (TCHAR)'(';
+					Writer << TEXT('(');
 				}
 				++Writer.ExpressionScope;
 				Expressions[0]->Write(Writer);
-				Writer << (TCHAR)' ';
+				Writer << TEXT(' ');
 				WriteOperator(Writer);
-				Writer << (TCHAR)' ';
+				Writer << TEXT(' ');
 				Expressions[1]->Write(Writer);
 				--Writer.ExpressionScope;
 				if (Writer.ExpressionScope != 0 && !IsAssignmentOperator(Operator))
 				{
-					Writer << (TCHAR)')';
+					Writer << TEXT(')');
 				}
 				break;
 			}
@@ -465,11 +441,11 @@ namespace CrossCompiler
 			default:
 				return false;
 
-			case EOperators::LogicOr:	OutValue = LHS || RHS; break;
-			case EOperators::LogicAnd:	OutValue = LHS && RHS; break;
+			case EOperators::LogicOr:	OutValue = LHS != 0 || RHS != 0; break;
+			case EOperators::LogicAnd:	OutValue = LHS != 0 && RHS != 0; break;
 			case EOperators::BitOr:		OutValue = LHS | RHS; break;
 			case EOperators::BitXor:	OutValue = LHS ^ RHS; break;
-			case EOperators::BitAnd:	OutValue = LHS ^ RHS; break;
+			case EOperators::BitAnd:	OutValue = LHS & RHS; break;
 			case EOperators::Equal:		OutValue = LHS == RHS; break;
 			case EOperators::NEqual:	OutValue = LHS != RHS; break;
 			case EOperators::Less:		OutValue = LHS < RHS; break;
@@ -502,10 +478,7 @@ namespace CrossCompiler
 
 		FExpressionStatement::~FExpressionStatement()
 		{
-			if (Expression)
-			{
-				delete Expression;
-			}
+			delete Expression;
 		}
 
 		FCompoundStatement::FCompoundStatement(FLinearAllocator* InAllocator, const FSourceInfo& InInfo) :
@@ -518,7 +491,7 @@ namespace CrossCompiler
 		{
 			Writer.DoIndent();
 			Writer << TEXT("{\n");
-			for (auto* Statement : Statements)
+			for (FNode* Statement : Statements)
 			{
 				FASTWriterIncrementScope Scope(Writer);
 				Statement->Write(Writer);
@@ -529,7 +502,7 @@ namespace CrossCompiler
 
 		FCompoundStatement::~FCompoundStatement()
 		{
-			for (auto* Statement : Statements)
+			for (FNode* Statement : Statements)
 			{
 				delete Statement;
 			}
@@ -662,7 +635,7 @@ namespace CrossCompiler
 
 		FFunction::~FFunction()
 		{
-			for (auto* Param : Parameters)
+			for (FNode* Param : Parameters)
 			{
 				delete Param;
 			}
@@ -711,10 +684,7 @@ namespace CrossCompiler
 
 		FJumpStatement::~FJumpStatement()
 		{
-			if (OptionalExpression)
-			{
-				delete OptionalExpression;
-			}
+			delete OptionalExpression;
 		}
 
 		FSelectionStatement::FSelectionStatement(FLinearAllocator* InAllocator, const FSourceInfo& InInfo) :
@@ -745,10 +715,7 @@ namespace CrossCompiler
 		{
 			delete Condition;
 			delete ThenStatement;
-			if (ElseStatement)
-			{
-				delete ElseStatement;
-			}
+			delete ElseStatement;
 		}
 
 		FTypeSpecifier::FTypeSpecifier(FLinearAllocator* InAllocator, const FSourceInfo& InInfo) :
@@ -759,7 +726,6 @@ namespace CrossCompiler
 			TextureMSNumSamples(1),
 			PatchSize(0),
 			bIsArray(false),
-			//bIsUnsizedArray(false),
 			ArraySize(nullptr)
 		{
 		}
@@ -779,44 +745,37 @@ namespace CrossCompiler
 				Writer << TypeName;
 				if (TextureMSNumSamples > 1)
 				{
-					Writer << (TCHAR)'<';
+					Writer << TEXT('<');
 					Writer << InnerType;
 					Writer << TEXT(", ");
 					Writer << (uint32)TextureMSNumSamples;
-					Writer << (TCHAR)'>';
+					Writer << TEXT('>');
 				}
 				else if (InnerType && *InnerType)
 				{
-					Writer << (TCHAR)'<';
+					Writer << TEXT('<');
 					Writer << InnerType;
-					Writer << (TCHAR)'>';
+					Writer << TEXT('>');
 				}
 			}
 
 			if (bIsArray)
 			{
-				Writer << TEXT("[ ");
+				Writer << TEXT('[');
 
 				if (ArraySize)
 				{
 					ArraySize->Write(Writer);
 				}
 
-				printf("]");
+				Writer << TEXT(']');
 			}
 		}
 
 		FTypeSpecifier::~FTypeSpecifier()
 		{
-			if (Structure)
-			{
-				delete Structure;
-			}
-
-			if (ArraySize)
-			{
-				delete ArraySize;
-			}
+			delete Structure;
+			delete ArraySize;
 		}
 
 		FCBufferDeclaration::FCBufferDeclaration(FLinearAllocator* InAllocator, const FSourceInfo& InInfo) :
@@ -831,12 +790,12 @@ namespace CrossCompiler
 			Writer.DoIndent();
 			Writer << TEXT("cbuffer ");
 			Writer << Name;
-			Writer << (TCHAR)'\n';
+			Writer << TEXT('\n');
 
 			Writer.DoIndent();
 			Writer << TEXT("{\n");
 
-			for (auto* Declaration : Declarations)
+			for (FNode* Declaration : Declarations)
 			{
 				FASTWriterIncrementScope Scope(Writer);
 				Declaration->Write(Writer);
@@ -848,7 +807,7 @@ namespace CrossCompiler
 
 		FCBufferDeclaration::~FCBufferDeclaration()
 		{
-			for (auto* Decl : Declarations)
+			for (FNode* Decl : Declarations)
 			{
 				delete Decl;
 			}
@@ -899,12 +858,12 @@ namespace CrossCompiler
 			if (NoInterpolatorType)
 			{
 				Writer << NoInterpolatorType;
-				Writer << (TCHAR)' ';
+				Writer << TEXT(' ');
 			}
 			if (NoPerspectiveType)
 			{
 				Writer << NoPerspectiveType;
-				Writer << (TCHAR)' ';
+				Writer << TEXT(' ');
 			}
 			if (bSample)
 			{
@@ -919,7 +878,7 @@ namespace CrossCompiler
 			if (PrimitiveType)
 			{
 				Writer << PrimitiveType;
-				Writer << (TCHAR)' ';
+				Writer << TEXT(' ');
 			}
 		}
 
@@ -1028,15 +987,11 @@ namespace CrossCompiler
 
 		FDeclaration::~FDeclaration()
 		{
-			for (auto* Expr : ArraySize)
+			for (FExpression* Expr : ArraySize)
 			{
 				delete Expr;
 			}
-
-			if (Initializer)
-			{
-				delete Initializer;
-			}
+			delete Initializer;
 		}
 
 		FDeclaratorList::FDeclaratorList(FLinearAllocator* InAllocator, const FSourceInfo& InInfo) :
@@ -1061,7 +1016,7 @@ namespace CrossCompiler
 			}
 
 			bool bFirst = true;
-			for (auto* Decl : Declarations)
+			for (FNode* Decl : Declarations)
 			{
 				if (bFirst)
 				{
@@ -1091,7 +1046,7 @@ namespace CrossCompiler
 		{
 			delete Type;
 
-			for (auto* Decl : Declarations)
+			for (FNode* Decl : Declarations)
 			{
 				delete Decl;
 			}
@@ -1112,7 +1067,7 @@ namespace CrossCompiler
 			default: break;
 			}
 			bool bFirst = true;
-			for (auto* Expr : Expressions)
+			for (FExpression* Expr : Expressions)
 			{
 				if (bFirst)
 				{
@@ -1148,7 +1103,7 @@ namespace CrossCompiler
 		{
 			WriteAttributes(Writer);
 			Type->Write(Writer);
-			Writer << (TCHAR)' ' << Identifier;
+			Writer << TEXT(' ') << Identifier;
 
 			WriteOptionArraySize(Writer, bIsArray, ArraySize);
 
@@ -1184,15 +1139,12 @@ namespace CrossCompiler
 		{
 			delete Type;
 
-			for (auto* Expr : ArraySize)
+			for (FExpression* Expr : ArraySize)
 			{
 				delete Expr;
 			}
 
-			if (DefaultValue)
-			{
-				delete DefaultValue;
-			}
+			delete DefaultValue;
 		}
 
 		FIterationStatement::FIterationStatement(FLinearAllocator* InAllocator, const FSourceInfo& InInfo, EIterationType InType) :
@@ -1215,7 +1167,7 @@ namespace CrossCompiler
 				Writer << TEXT("for (");
 				if (InitStatement)
 				{
-					auto* DeclList = InitStatement->AsDeclaratorList();
+					FDeclaratorList* DeclList = InitStatement->AsDeclaratorList();
 					if (DeclList)
 					{
 						DeclList->WriteNoEOL(Writer);
@@ -1295,24 +1247,10 @@ namespace CrossCompiler
 
 		FIterationStatement::~FIterationStatement()
 		{
-			if (InitStatement)
-			{
-				delete InitStatement;
-			}
-			if (Condition)
-			{
-				delete Condition;
-			}
-
-			if (RestExpression)
-			{
-				delete RestExpression;
-			}
-
-			if (Body)
-			{
-				delete Body;
-			}
+			delete InitStatement;
+			delete Condition;
+			delete RestExpression;
+			delete Body;
 		}
 
 		FFunctionExpression::FFunctionExpression(FLinearAllocator* InAllocator, const FSourceInfo& InInfo, FExpression* InCallee) :
@@ -1324,9 +1262,9 @@ namespace CrossCompiler
 		void FFunctionExpression::Write(FASTWriter& Writer) const
 		{
 			Callee->Write(Writer);
-			Writer << (TCHAR)'(';
+			Writer << TEXT('(');
 			bool bFirst = true;
-			for (auto* Expr : Expressions)
+			for (FExpression* Expr : Expressions)
 			{
 				if (bFirst)
 				{
@@ -1410,10 +1348,7 @@ namespace CrossCompiler
 
 		FCaseLabel::~FCaseLabel()
 		{
-			if (TestExpression)
-			{
-				delete TestExpression;
-			}
+			delete TestExpression;
 		}
 
 
@@ -1450,12 +1385,9 @@ namespace CrossCompiler
 		FCaseStatement::~FCaseStatement()
 		{
 			delete Labels;
-			for (auto* Statement : Statements)
+			for (FNode* Statement : Statements)
 			{
-				if (Statement)
-				{
-					delete Statement;
-				}
+				delete Statement;
 			}
 		}
 
@@ -1467,7 +1399,7 @@ namespace CrossCompiler
 
 		void FCaseLabelList::Write(FASTWriter& Writer) const
 		{
-			for (auto* Label : Labels)
+			for (FCaseLabel* Label : Labels)
 			{
 				Label->Write(Writer);
 			}
@@ -1475,7 +1407,7 @@ namespace CrossCompiler
 
 		FCaseLabelList::~FCaseLabelList()
 		{
-			for (auto* Label : Labels)
+			for (FCaseLabel* Label : Labels)
 			{
 				delete Label;
 			}
@@ -1489,7 +1421,7 @@ namespace CrossCompiler
 
 		void FCaseStatementList::Write(FASTWriter& Writer) const
 		{
-			for (auto* Case : Cases)
+			for (FCaseStatement* Case : Cases)
 			{
 				Case->Write(Writer);
 			}
@@ -1497,7 +1429,7 @@ namespace CrossCompiler
 
 		FCaseStatementList::~FCaseStatementList()
 		{
-			for (auto* Case : Cases)
+			for (FCaseStatement* Case : Cases)
 			{
 				delete Case;
 			}
@@ -1531,7 +1463,7 @@ namespace CrossCompiler
 				Writer.DoIndent();
 				Writer << TEXT("{\n");
 
-				for (auto* Member : Members)
+				for (FNode* Member : Members)
 				{
 					FASTWriterIncrementScope Scope(Writer);
 					Member->Write(Writer);
@@ -1544,7 +1476,7 @@ namespace CrossCompiler
 
 		FStructSpecifier::~FStructSpecifier()
 		{
-			for (auto* Member : Members)
+			for (FNode* Member : Members)
 			{
 				delete Member;
 			}
@@ -1563,11 +1495,11 @@ namespace CrossCompiler
 			Writer << Name;
 
 			bool bFirst = true;
-			for (auto* Arg : Arguments)
+			for (FAttributeArgument* Arg : Arguments)
 			{
 				if (bFirst)
 				{
-					Writer << (TCHAR)'(';
+					Writer << TEXT('(');
 					bFirst = false;
 				}
 				else
@@ -1588,7 +1520,7 @@ namespace CrossCompiler
 
 		FAttribute::~FAttribute()
 		{
-			for (auto* Arg : Arguments)
+			for (FAttributeArgument* Arg : Arguments)
 			{
 				delete Arg;
 			}
@@ -1609,18 +1541,15 @@ namespace CrossCompiler
 			}
 			else
 			{
-				Writer << (TCHAR)'"';
+				Writer << TEXT('"');
 				Writer << StringArgument;
-				Writer << (TCHAR)'"';
+				Writer << TEXT('"');
 			}
 		}
 
 		FAttributeArgument::~FAttributeArgument()
 		{
-			if (ExpressionArgument)
-			{
-				delete ExpressionArgument;
-			}
+			delete ExpressionArgument;
 		}
 	}
 }
diff --git a/Engine/Source/Developer/HlslParser/Private/HlslAST.h b/Engine/Source/Developer/HlslParser/Private/HlslAST.h
--- a/Engine/Source/Developer/HlslParser/Private/HlslAST.h
+++ b/Engine/Source/Developer/HlslParser/Private/HlslAST.h
@@ -102,20 +102,17 @@ namespace CrossCompiler
 		struct FASTWriter
 		{
 			FString& Output;
-			int32 Indent;
-			int32 ExpressionScope;
+			int32 Indent = 0;
+			int32 ExpressionScope = 0;
 
 			FASTWriter(FString& FinalOutput) :
-				Output(FinalOutput),
-				Indent(0),
-				ExpressionScope(0)
+				Output(FinalOutput)
 			{
 			}
 
 			// Construct from another to go back to unindented writing
 			FASTWriter(FASTWriter& IndentedWriter) :
-				Output(IndentedWriter.Output),
-				Indent(0)
+				Output(IndentedWriter.Output)
 			{
 			}
 
@@ -211,7 +208,7 @@ namespace CrossCompiler
 		struct FASTWriterSkipExpressionScope
 		{
 			FASTWriter& Writer;
-			int32 OriginalScope;
+			int32 OriginalScope = 0;
 
 			FASTWriterSkipExpressionScope(FASTWriter& InWriter) :
 				Writer(InWriter)
@@ -229,7 +226,7 @@ namespace CrossCompiler
 		struct FASTWriterClearIndentScope
 		{
 			FASTWriter& Writer;
-			int32 OriginalIndent;
+			int32 OriginalIndent = 0;
 
 			FASTWriterClearIndentScope(FASTWriter& InWriter) :
 				Writer(InWriter)
@@ -609,7 +606,6 @@ namespace CrossCompiler
 			FSemanticSpecifier* Semantic;
 
 			bool bIsArray;
-			//bool bIsUnsizedArray;
 
 			TLinearArray<FExpression*> ArraySize;
 
@@ -689,7 +685,6 @@ namespace CrossCompiler
 			int32 PatchSize;
 			bool bPrecise = false;
 			bool bIsArray;
-			//bool bIsUnsizedArray;
 			FExpression* ArraySize;
 		};
 
diff --git a/Engine/Source/Developer/HlslParser/Private/HlslExpressionParser.inl b/Engine/Source/Developer/HlslParser/Private/HlslExpressionParser.inl
--- a/Engine/Source/Developer/HlslParser/Private/HlslExpressionParser.inl
+++ b/Engine/Source/Developer/HlslParser/Private/HlslExpressionParser.inl
@@ -1186,7 +1186,6 @@ namespace CrossCompiler
 				return 2;
 
 			case EHlslToken::Question:
-				check(0);
 				return 3;
 
 			case EHlslToken::OrOr:
@@ -1321,15 +1320,15 @@ namespace CrossCompiler
 		check(*OutExpression);
 		do
 		{
-			auto* Token = Scanner.GetCurrentToken();
-			int32 Precedence = GetPrecedence(Token);
+			const FHlslToken* Token = Scanner.GetCurrentToken();
+			const int32 Precedence = GetPrecedence(Token);
 			if (!Token || !IsBinaryOperator(Token) || Precedence < MinPrec || (!bAllowAssignment && IsAssignmentOperator(Token)) || IsSequenceOperator(Token) || (OutTernaryExpression && *OutTernaryExpression))
 			{
 				break;
 			}
 
 			Scanner.Advance();
-			auto NextMinPrec = IsRightAssociative(Token) ? Precedence : Precedence + 1;
+			const int32 NextMinPrec = IsRightAssociative(Token) ? Precedence : Precedence + 1;
 			AST::FExpression* RHSExpression = nullptr;
 			AST::FExpression* RHSTernaryExpression = nullptr;
 			Result = ComputeExpr(Scanner, NextMinPrec, /*Info,*/ SymbolScope, ExpressionFlags, Allocator, &RHSExpression, &RHSTernaryExpression);
@@ -1342,7 +1341,7 @@ namespace CrossCompiler
 				break;
 			}
 			check(RHSExpression);
-			auto BinaryOperator = AST::TokenToASTOperator(Token->Token);
+			const AST::EOperators BinaryOperator = AST::TokenToASTOperator(Token->Token);
 			*OutExpression = new(Allocator) AST::FBinaryExpression(Allocator, BinaryOperator, *OutExpression, RHSExpression, Token->SourceInfo);
 			if  (RHSTernaryExpression)
 			{
@@ -1464,7 +1463,7 @@ namespace CrossCompiler
 			else if (NewStartListToken != EHlslToken::Invalid && Token->Token == NewStartListToken)
 			{
 				Scanner.Advance();
-				check(0);
+				checkNoEntry();
 /*
 				auto* SubExpression = new(Allocator) AST::FInitializerListExpression(Allocator, Token->SourceInfo);
 				auto Result = ParseExpressionList(EndListToken, Scanner, SymbolScope, NewStartListToken, Allocator, SubExpression);
diff --git a/Engine/Source/Developer/HlslParser/Private/HlslLexer.cpp b/Engine/Source/Developer/HlslParser/Private/HlslLexer.cpp
--- a/Engine/Source/Developer/HlslParser/Private/HlslLexer.cpp
+++ b/Engine/Source/Developer/HlslParser/Private/HlslLexer.cpp
@@ -1125,6 +1125,7 @@ namespace CrossCompiler
 	void FHlslScanner::Clear(const FString& Filename)
 	{
 		Tokens.Empty();
+		SourceFilenames.Empty();
 		SourceFilenames.Add(new FString(Filename));
 	}
 
@@ -1273,18 +1274,6 @@ namespace CrossCompiler
 		{
 			return &Tokens[CurrentToken];
 		}
-
-		return nullptr;
-	}
-
-	const FHlslToken* FHlslScanner::GetCurrentTokenAndAdvance()
-	{
-		if (CurrentToken < (uint32)Tokens.Num())
-		{
-			auto* Return = &Tokens[CurrentToken];
-			Advance();
-		}
-
 		return nullptr;
 	}
 
@@ -1298,7 +1287,7 @@ namespace CrossCompiler
 	{
 		if (CurrentToken < (uint32)Tokens.Num())
 		{
-			const auto& Token = Tokens[CurrentToken];
+			const FHlslToken& Token = Tokens[CurrentToken];
 			check(Token.SourceInfo.Filename);
 			CompilerMessages.SourceError(Token.SourceInfo, *Error);
 		}
@@ -1310,7 +1299,7 @@ namespace CrossCompiler
 
 	void FTokenizer::ProcessDirective(FTokenizer& Tokenizer, FCompilerMessages& CompilerMessages, FHlslScanner& Scanner)
 	{
-		check(Tokenizer.Peek() == '#');
+		check(Tokenizer.Peek() == TEXT('#'));
 		if (Tokenizer.MatchString(MATCH_TARGET(TEXT("#line"))))
 		{
 			Tokenizer.SkipWhitespaceInLine();
diff --git a/Engine/Source/Developer/HlslParser/Private/HlslLexer.h b/Engine/Source/Developer/HlslParser/Private/HlslLexer.h
--- a/Engine/Source/Developer/HlslParser/Private/HlslLexer.h
+++ b/Engine/Source/Developer/HlslParser/Private/HlslLexer.h
@@ -388,7 +388,6 @@ namespace CrossCompiler
 		bool MatchIntegerLiteral();
 		const FHlslToken* PeekToken(uint32 LookAhead = 0) const;
 		const FHlslToken* GetCurrentToken() const;
-		const FHlslToken* GetCurrentTokenAndAdvance();
 		void Advance() {++CurrentToken;}
 
 		void SourceError(const FString& Error);
diff --git a/Engine/Source/Developer/HlslParser/Private/HlslUtils.cpp b/Engine/Source/Developer/HlslParser/Private/HlslUtils.cpp
--- a/Engine/Source/Developer/HlslParser/Private/HlslUtils.cpp
+++ b/Engine/Source/Developer/HlslParser/Private/HlslUtils.cpp
@@ -113,6 +113,7 @@ static inline TCHAR* GetNonDigitSemanticPrefix(CrossCompiler::FLinearAllocator*
 
 static inline TCHAR* MakeIndexedSemantic(CrossCompiler::FLinearAllocator* Allocator, const TCHAR* Semantic, uint32 Index)
 {
+	check(Semantic);
 	FString Out = FString::Printf(TEXT("%s%d"), Semantic, Index);
 	return Allocator->Strdup(Out);
 }
@@ -840,33 +841,35 @@ struct FRemoveUnusedOutputs : FRemoveAlgorithm
 
 								CopyMember(MemberDeclaration, DestPrefix, SourcePrefix, BodyContext.PostInstructions);
 							}
-
-							for (uint32 Index = 0; Index < ArrayLength; ++Index)
+							else
 							{
-								TCHAR* ElementSemantic = MakeIndexedSemantic(Allocator, ElementSemanticPrefix, StartIndex + Index);
-								if (IsSemanticUsed(ElementSemantic))
-								{
-									auto* NewMemberDeclaration = new(Allocator) FDeclaration(Allocator, MemberDeclaration->SourceInfo);
-									NewMemberDeclaration->Semantic = new(Allocator) FSemanticSpecifier(Allocator, ElementSemantic, MemberDeclaration->SourceInfo);
-									NewMemberDeclaration->Identifier = FIdentifier::New(Allocator, FString::Printf(TEXT("%s_%d"), MemberDeclaration->Identifier->GetData(), Index));
-
-									// Add member to struct
-									auto* NewDeclaratorList = new(Allocator) FDeclaratorList(Allocator, MemberDeclarator->SourceInfo);
-									NewDeclaratorList->Type = CloneType(MemberDeclarator->Type);
-									NewDeclaratorList->Declarations.Add(NewMemberDeclaration);
-									DestStruct->Members.Add(NewDeclaratorList);
-
-									FString LHSElement = FString::Printf(TEXT("%s.%s"), DestPrefix, NewMemberDeclaration->Identifier->GetData());
-									FString RHSElement = FString::Printf(TEXT("%s.%s[%d]"), SourcePrefix, MemberDeclaration->Identifier->GetData(), Index);
-
-									auto* LHS = MakeIdentifierExpression(Allocator, *LHSElement, SourceInfo);
-									auto* RHS = MakeIdentifierExpression(Allocator, *RHSElement, SourceInfo);
-									auto* Assignment = new(Allocator) FBinaryExpression(Allocator, EOperators::Assign, LHS, RHS, SourceInfo);
-									BodyContext.PostInstructions.Add(new(Allocator) FExpressionStatement(Allocator, Assignment, SourceInfo));
-								}
-								else
+								for (uint32 Index = 0; Index < ArrayLength; ++Index)
 								{
-									RemovedSemantics.Add(ElementSemantic);
+									TCHAR* ElementSemantic = MakeIndexedSemantic(Allocator, ElementSemanticPrefix, StartIndex + Index);
+									if (IsSemanticUsed(ElementSemantic))
+									{
+										auto* NewMemberDeclaration = new(Allocator) FDeclaration(Allocator, MemberDeclaration->SourceInfo);
+										NewMemberDeclaration->Semantic = new(Allocator) FSemanticSpecifier(Allocator, ElementSemantic, MemberDeclaration->SourceInfo);
+										NewMemberDeclaration->Identifier = FIdentifier::New(Allocator, FString::Printf(TEXT("%s_%d"), MemberDeclaration->Identifier->GetData(), Index));
+
+										// Add member to struct
+										auto* NewDeclaratorList = new(Allocator) FDeclaratorList(Allocator, MemberDeclarator->SourceInfo);
+										NewDeclaratorList->Type = CloneType(MemberDeclarator->Type);
+										NewDeclaratorList->Declarations.Add(NewMemberDeclaration);
+										DestStruct->Members.Add(NewDeclaratorList);
+
+										FString LHSElement = FString::Printf(TEXT("%s.%s"), DestPrefix, NewMemberDeclaration->Identifier->GetData());
+										FString RHSElement = FString::Printf(TEXT("%s.%s[%d]"), SourcePrefix, MemberDeclaration->Identifier->GetData(), Index);
+
+										auto* LHS = MakeIdentifierExpression(Allocator, *LHSElement, SourceInfo);
+										auto* RHS = MakeIdentifierExpression(Allocator, *RHSElement, SourceInfo);
+										auto* Assignment = new(Allocator) FBinaryExpression(Allocator, EOperators::Assign, LHS, RHS, SourceInfo);
+										BodyContext.PostInstructions.Add(new(Allocator) FExpressionStatement(Allocator, Assignment, SourceInfo));
+									}
+									else
+									{
+										RemovedSemantics.Add(ElementSemantic);
+									}
 								}
 							}
 						}
@@ -1698,7 +1701,11 @@ bool ForEachHlslIdentifier(const FString& InSourceCode, const FString& InFilenam
 		{
 			if (HlslToken->Token == CrossCompiler::EHlslToken::Identifier)
 			{
-				InIdentifierCallback(HlslToken->String);
+				// Report identifier to callback and stop iterating once the callback returns false
+				if (!InIdentifierCallback(HlslToken->String))
+				{
+					break;
+				}
 			}
 		}
 		Scanner.Advance();
@@ -2044,8 +2051,6 @@ bool PrettyPrintHlslParser(FString& InOutSourceCode, TArray<FString>& OutErrors)
 	for (auto& Message : Messages.MessageList)
 	{
 		OutErrors.Add(Message.Message);
-		//#todo-rco: 
-		check(0);
 	}
 
 	InOutSourceCode = Out;
