[HlslParser] Fix call of MakeIndexedSemantic() with null pointer for semantic prefix. Follow-up to CL 45280838.

- This also handles the second case of calling MakeIndexedSemantic() with a null pointer.
- Remove unused function PrettyPrintHlslParser().

#rnx

diff --git a/Engine/Source/Developer/HlslParser/Private/HlslUtils.cpp b/Engine/Source/Developer/HlslParser/Private/HlslUtils.cpp
--- a/Engine/Source/Developer/HlslParser/Private/HlslUtils.cpp
+++ b/Engine/Source/Developer/HlslParser/Private/HlslUtils.cpp
@@ -1307,33 +1307,35 @@ struct FRemoveUnusedInputs : FRemoveAlgorithm
 								// Source and Dest are swapped as we are copying from the optimized (dest) structure into the original (source) structure
 								CopyMember(MemberDeclaration, SourcePrefix, DestPrefix, BodyContext.PreInstructions);
 							}
-
-							for (uint32 Index = 0; Index < ArrayLength; ++Index)
+							else
 							{
-								TCHAR* ElementSemantic = MakeIndexedSemantic(Allocator, ElementSemanticPrefix, StartIndex + Index);
-								if (IsSemanticUsed(ElementSemantic))
+								for (uint32 Index = 0; Index < ArrayLength; ++Index)
 								{
-									auto* NewMemberDeclaration = new(Allocator) FDeclaration(Allocator, MemberDeclaration->SourceInfo);
-									NewMemberDeclaration->Semantic = new(Allocator) FSemanticSpecifier(Allocator, ElementSemantic, MemberDeclaration->SourceInfo);
-									NewMemberDeclaration->Identifier = FIdentifier::New(Allocator, FString::Printf(TEXT("%s_%d"), MemberDeclaration->Identifier->GetData(), Index));
-
-									// Add member to struct
-									auto* NewDeclaratorList = new(Allocator) FDeclaratorList(Allocator, MemberDeclarator->SourceInfo);
-									NewDeclaratorList->Type = CloneType(MemberDeclarator->Type);
-									NewDeclaratorList->Declarations.Add(NewMemberDeclaration);
-									DestStruct->Members.Add(NewDeclaratorList);
-
-									FString LHSElement = FString::Printf(TEXT("%s.%s[%d]"), SourcePrefix, MemberDeclaration->Identifier->GetData(), Index);
-									FString RHSElement = FString::Printf(TEXT("%s.%s"), DestPrefix, NewMemberDeclaration->Identifier->GetData());
-
-									auto* LHS = MakeIdentifierExpression(Allocator, *LHSElement, SourceInfo);
-									auto* RHS = MakeIdentifierExpression(Allocator, *RHSElement, SourceInfo);
-									auto* Assignment = new(Allocator) FBinaryExpression(Allocator, EOperators::Assign, LHS, RHS, SourceInfo);
-									BodyContext.PreInstructions.Add(new(Allocator) FExpressionStatement(Allocator, Assignment, SourceInfo));
-								}
-								else
-						{
-									RemovedSemantics.Add(ElementSemantic);
+									TCHAR* ElementSemantic = MakeIndexedSemantic(Allocator, ElementSemanticPrefix, StartIndex + Index);
+									if (IsSemanticUsed(ElementSemantic))
+									{
+										auto* NewMemberDeclaration = new(Allocator) FDeclaration(Allocator, MemberDeclaration->SourceInfo);
+										NewMemberDeclaration->Semantic = new(Allocator) FSemanticSpecifier(Allocator, ElementSemantic, MemberDeclaration->SourceInfo);
+										NewMemberDeclaration->Identifier = FIdentifier::New(Allocator, FString::Printf(TEXT("%s_%d"), MemberDeclaration->Identifier->GetData(), Index));
+
+										// Add member to struct
+										auto* NewDeclaratorList = new(Allocator) FDeclaratorList(Allocator, MemberDeclarator->SourceInfo);
+										NewDeclaratorList->Type = CloneType(MemberDeclarator->Type);
+										NewDeclaratorList->Declarations.Add(NewMemberDeclaration);
+										DestStruct->Members.Add(NewDeclaratorList);
+
+										FString LHSElement = FString::Printf(TEXT("%s.%s[%d]"), SourcePrefix, MemberDeclaration->Identifier->GetData(), Index);
+										FString RHSElement = FString::Printf(TEXT("%s.%s"), DestPrefix, NewMemberDeclaration->Identifier->GetData());
+
+										auto* LHS = MakeIdentifierExpression(Allocator, *LHSElement, SourceInfo);
+										auto* RHS = MakeIdentifierExpression(Allocator, *RHSElement, SourceInfo);
+										auto* Assignment = new(Allocator) FBinaryExpression(Allocator, EOperators::Assign, LHS, RHS, SourceInfo);
+										BodyContext.PreInstructions.Add(new(Allocator) FExpressionStatement(Allocator, Assignment, SourceInfo));
+									}
+									else
+									{
+										RemovedSemantics.Add(ElementSemantic);
+									}
 								}
 							}
 						}
@@ -2022,41 +2024,6 @@ bool ConvertFromFP32ToFP16(FString& InOutSourceCode, TArray<FString>& OutErrors)
 	return false;
 }
 
-
-bool PrettyPrintHlslParser(FString& InOutSourceCode, TArray<FString>& OutErrors)
-{
-	auto PrettyPrinter = [](void* CallbackData, CrossCompiler::FLinearAllocator* Allocator, CrossCompiler::TLinearArray<CrossCompiler::AST::FNode*>& ASTNodes)
-	{
-		FString* Out = (FString*)CallbackData;
-		CrossCompiler::AST::FASTWriter Writer(*Out);
-		for (CrossCompiler::AST::FNode* Node : ASTNodes)
-		{
-			Node->Write(Writer);
-		}
-	};
-
-	FString DummyFilename(TEXT("/Engine/Private/PrettyPrinter.usf"));
-	FString Out;
-	CrossCompiler::FCompilerMessages Messages;
-	if (!CrossCompiler::Parser::Parse(InOutSourceCode, DummyFilename, Messages, {}, PrettyPrinter, &Out))
-	{
-		OutErrors.Add(FString(TEXT("PrettyPrintHlslParser: Failed to compile!")));
-		for (auto& Message : Messages.MessageList)
-		{
-			OutErrors.Add(Message.Message);
-		}
-		return false;
-	}
-
-	for (auto& Message : Messages.MessageList)
-	{
-		OutErrors.Add(Message.Message);
-	}
-
-	InOutSourceCode = Out;
-	return true;
-}
-
 namespace UE::ShaderCompilerCommon
 {
 	static void BuildDiagnosticsSourceInfoMsvc(TStringBuilder<1024>& DiagnosticsStringBuilder, const CrossCompiler::FSourceInfo& SourceInfo)
