[Shaders] Print more look-ahead tokens in HlslParser error when no rule can be matched.

#rnx

diff --git a/Engine/Source/Developer/HlslParser/Private/HlslParser.cpp b/Engine/Source/Developer/HlslParser/Private/HlslParser.cpp
--- a/Engine/Source/Developer/HlslParser/Private/HlslParser.cpp
+++ b/Engine/Source/Developer/HlslParser/Private/HlslParser.cpp
@@ -2266,6 +2266,22 @@ Done:
 		return false;
 	}
 
+	// Prints the specified number of tokens, starting from the token the scanner is currently pointing to, into a string.
+	static FString PrintDiagnosticTokenStream(const FHlslScanner& Scanner, int32 Count)
+	{
+		const FHlslToken* Token = nullptr;
+		FString OutString;
+		for (int32 TokenIndex = 0; TokenIndex < Count && (Token = Scanner.PeekToken(TokenIndex)) != nullptr; ++TokenIndex)
+		{
+			if (TokenIndex > 0)
+			{
+				OutString += TCHAR(' ');
+			}
+			OutString += Token->String;
+		}
+		return OutString;
+	}
+
 	EParseResult TryTranslationUnit(FHlslParser& Parser, FLinearAllocator* Allocator, AST::FNode** OutNode)
 	{
 		if (MatchPragma(Parser, Allocator, OutNode))
@@ -2353,7 +2369,8 @@ Done:
 			return EParseResult::Ignore;
 		}
 
-		Parser.Scanner.SourceError(FString::Printf(TEXT("Unable to match rule! [Token=\"%s\"]"), *Peek->String));
+		constexpr int32 kTokenLookAheadForError = 10;
+		Parser.Scanner.SourceError(FString::Printf(TEXT("Unable to match rule! [Token stream: \"%s\"]"), *PrintDiagnosticTokenStream(Parser.Scanner, kTokenLookAheadForError)));
 		return ParseResultError();
 	}
 
