Fix HLSL parser not supporting a bunch of features (struct bitpacking, templated function call, functions with numeric type and usertype arguments)
Reenable RemoveUnusedInterpolators on some platform

diff --git a/Engine/Shaders/Public/ShaderVersion.ush b/Engine/Shaders/Public/ShaderVersion.ush
--- a/Engine/Shaders/Public/ShaderVersion.ush
+++ b/Engine/Shaders/Public/ShaderVersion.ush
@@ -3,4 +3,4 @@
 // Each time the console command is executed it generates a new GUID. As this file is included
 // in Platform.ush (which should be included in any shader) it allows to invalidate the shader DDC.
 // 
-#pragma message("UESHADERMETADATA_VERSION 9D55C75A-B0F6-48ED-AA2F-92EBB73B3ACE")
\ No newline at end of file
+#pragma message("UESHADERMETADATA_VERSION 1EBEB311-1396-4095-A654-42314B22EFD0")
\ No newline at end of file
diff --git a/Engine/Source/Developer/HlslParser/Private/HlslAST.cpp b/Engine/Source/Developer/HlslParser/Private/HlslAST.cpp
--- a/Engine/Source/Developer/HlslParser/Private/HlslAST.cpp
+++ b/Engine/Source/Developer/HlslParser/Private/HlslAST.cpp
@@ -275,6 +275,10 @@ namespace CrossCompiler
 				Writer << TEXT(')');
 				break;
 
+			case EOperators::SizeOf:
+				Writer << TEXT("sizeof");
+				break;
+
 			default:
 				Writer << TEXT("*MISSING_");
 				Writer << (uint32)Operator;
@@ -1064,6 +1068,7 @@ namespace CrossCompiler
 			{
 			case EType::Braced:			Writer << TEXT("{"); break;
 			case EType::Parenthesized:	Writer << TEXT("("); break;
+			case EType::Templated:		Writer << TEXT("<"); break;
 			default: break;
 			}
 			bool bFirst = true;
@@ -1084,6 +1089,7 @@ namespace CrossCompiler
 			{
 			case EType::Braced:			Writer << TEXT("}"); break;
 			case EType::Parenthesized:	Writer << TEXT(")"); break;
+			case EType::Templated:		Writer << TEXT(">"); break;
 			default: break;
 			}
 		}
diff --git a/Engine/Source/Developer/HlslParser/Private/HlslAST.h b/Engine/Source/Developer/HlslParser/Private/HlslAST.h
--- a/Engine/Source/Developer/HlslParser/Private/HlslAST.h
+++ b/Engine/Source/Developer/HlslParser/Private/HlslAST.h
@@ -356,6 +356,7 @@ namespace CrossCompiler
 			Literal,
 
 			TypeCast,
+			SizeOf,
 		};
 
 		inline EOperators TokenToASTOperator(EHlslToken Token)
@@ -531,6 +532,7 @@ namespace CrossCompiler
 				FreeForm,			// a, b, c
 				Parenthesized,		// ( a, b, c )
 				Braced,				// { a, b, c }
+				Templated,			// < a, b, c >
 			};
 
 			FExpressionList(FLinearAllocator* InAllocator, EType InType, const FSourceInfo& InInfo);
diff --git a/Engine/Source/Developer/HlslParser/Private/HlslExpressionParser.inl b/Engine/Source/Developer/HlslParser/Private/HlslExpressionParser.inl
--- a/Engine/Source/Developer/HlslParser/Private/HlslExpressionParser.inl
+++ b/Engine/Source/Developer/HlslParser/Private/HlslExpressionParser.inl
@@ -24,7 +24,7 @@ namespace CrossCompiler
 
 	EParseResult ComputeExpr(FHlslScanner& Scanner, int32 MinPrec, /*FInfo& Info,*/ FSymbolScope* SymbolScope, int32 ExpressionFlags, FLinearAllocator* Allocator, AST::FExpression** OutExpression, AST::FExpression** OutTernaryExpression);
 	EParseResult ParseExpressionList(EHlslToken EndListToken, FHlslScanner& Scanner, FSymbolScope* SymbolScope, EHlslToken NewStartListToken, FLinearAllocator* Allocator, AST::FExpression* OutExpression);
-
+	EParseResult ParseExpression(FHlslScanner& Scanner, FSymbolScope* SymbolScope, int32 ExpressionFlags, FLinearAllocator* Allocator, AST::FExpression** OutExpression);
 	struct FSymbolScope
 	{
 		enum class EType
@@ -910,18 +910,64 @@ namespace CrossCompiler
 			{
 				Scanner.Advance();
 
-				// Function Call
-				auto* FunctionCall = new(Allocator) AST::FFunctionExpression(Allocator, Token->SourceInfo, PrevExpression);
-				auto Result = ParseExpressionList(EHlslToken::RightParenthesis, Scanner, SymbolScope, EHlslToken::Invalid, Allocator, FunctionCall);
-				if (Result != EParseResult::Matched)
+				if (PrevExpression->Identifier && PrevExpression->Identifier->ToString() == TEXT("sizeof"))
 				{
-					Scanner.SourceError(TEXT("Expected ')'!"));
-					return ParseResultError();
+					const FHlslToken* Token1 = Scanner.PeekToken();
+
+					AST::FTypeSpecifier* TypeSpecifier = nullptr;
+					auto* SizeOfExpression = new(Allocator) AST::FUnaryExpression(Allocator, AST::EOperators::SizeOf, PrevExpression, Token1->SourceInfo);
+
+					if (ParseGeneralTypeFromToken(Token1, ETF_BUILTIN_NUMERIC, 0, SymbolScope, Allocator, &TypeSpecifier) == EParseResult::Matched)
+					{
+						AST::FTypeSpecifier* ElementTypeSpecifier = nullptr;
+						EParseResult InnerResult = ParseGeneralType(Scanner, ETF_BUILTIN_NUMERIC, SymbolScope, Allocator, &ElementTypeSpecifier);
+						if (InnerResult != EParseResult::Matched)
+						{
+							Scanner.SourceError(TEXT("Expected numeric type!"));
+							return ParseResultError();
+						}
+
+						SizeOfExpression->Identifier = AST::FIdentifier::New(Allocator, ElementTypeSpecifier->TypeName);
+					}
+					else
+					{
+						AST::FExpression* ConditionExpression = nullptr;
+						auto Result = ParseExpression(Scanner, SymbolScope, 0, Allocator, &ConditionExpression);
+						if (Result == EParseResult::Error)
+						{
+							Scanner.SourceError(TEXT("Expected expression or declaration!\n"));
+							return ParseResultError();
+						}
+
+						SizeOfExpression->Identifier = AST::FIdentifier::New(Allocator, ConditionExpression->Identifier->ToString());
+					}
+
+					if (!Scanner.MatchToken(EHlslToken::RightParenthesis))
+					{
+						Scanner.SourceError(TEXT("Expected )"));
+						return ParseResultError();
+					}
+
+					PrevExpression = SizeOfExpression;
+					bFoundAny = true;
+				}
+				else
+				{
+					// Function Call
+					auto* FunctionCall = new(Allocator) AST::FFunctionExpression(Allocator, Token->SourceInfo, PrevExpression);
+					auto Result = ParseExpressionList(EHlslToken::RightParenthesis, Scanner, SymbolScope, EHlslToken::Invalid, Allocator, FunctionCall);
+					if (Result != EParseResult::Matched)
+					{
+						Scanner.SourceError(TEXT("Expected ')'!"));
+						return ParseResultError();
+					}
+
+					PrevExpression = FunctionCall;
+
+					bFoundAny = true;
 				}
 
-				PrevExpression = FunctionCall;
 
-				bFoundAny = true;
 			}
 				break;
 			case EHlslToken::PlusPlus:
@@ -965,6 +1011,68 @@ namespace CrossCompiler
 				*OutTernaryExpression = Ternary;
 				bFoundAny = true;
 				bTryAgain = false;
+			}
+				break;
+			case EHlslToken::Lower:
+			{
+				if (SymbolScope->ScopeType == FSymbolScope::EType::Statement)
+				{
+					const FHlslToken* Token1 = Scanner.PeekToken(1);
+					
+					AST::FTypeSpecifier* TypeSpecifier = nullptr;
+					if (ParseGeneralTypeFromToken(Token1, ETF_BUILTIN_NUMERIC | ETF_USER_TYPES, 0, SymbolScope, Allocator, &TypeSpecifier) == EParseResult::Matched )
+					{
+						Scanner.Advance();
+						
+						AST::FExpressionList* ExpressionList = new(Allocator) AST::FExpressionList(Allocator, AST::FExpressionList::EType::Templated, Scanner.PeekToken()->SourceInfo);
+
+						while (Scanner.HasMoreTokens())
+						{
+							const auto* CurrentToken = Scanner.PeekToken();
+							if (CurrentToken->Token == EHlslToken::Greater)
+							{
+								Scanner.Advance();
+								bFoundAny = true;
+								break;
+							}
+
+							if (CurrentToken->Token == EHlslToken::Comma)
+							{
+								Scanner.Advance();
+							}
+							else
+							{
+								AST::FTypeSpecifier* ElementTypeSpecifier = nullptr;
+								EParseResult InnerResult = ParseGeneralType(Scanner, ETF_BUILTIN_NUMERIC | ETF_USER_TYPES, SymbolScope, Allocator, &ElementTypeSpecifier);
+								if (InnerResult != EParseResult::Matched)
+								{
+									Scanner.SourceError(TEXT("Expected numeric or user type!"));
+									return ParseResultError();
+								}
+
+								AST::FExpression* CurrentExpression = new(Allocator) AST::FExpression(Allocator, AST::EOperators::Identifier, CurrentToken->SourceInfo);
+								ExpressionList->Expressions.Add(CurrentExpression);
+							}
+						}
+
+						if (!bFoundAny)
+						{
+							Scanner.SourceError(TEXT("Expected >"));
+							return ParseResultError();
+						}
+
+						PrevExpression = ExpressionList;
+					}
+					else
+					{
+						bTryAgain = false;
+					}
+				}
+				else
+				{
+					bTryAgain = false;
+				}
+
 			}
 				break;
 			default:
diff --git a/Engine/Source/Developer/HlslParser/Private/HlslParser.cpp b/Engine/Source/Developer/HlslParser/Private/HlslParser.cpp
--- a/Engine/Source/Developer/HlslParser/Private/HlslParser.cpp
+++ b/Engine/Source/Developer/HlslParser/Private/HlslParser.cpp
@@ -193,6 +193,7 @@ namespace CrossCompiler
 		EDF_INTERPOLATION				= (1 << 12),
 		EDF_REGISTER					= (1 << 13),
 		EDF_PACKOFFSET					= (1 << 14),
+		EDF_BITPACK					= (1 << 15),
 	};
 
 	EParseResult ParseInitializer(FHlslScanner& Scanner, FSymbolScope* SymbolScope, bool bAllowLists, FLinearAllocator* Allocator, AST::FExpression** OutList)
@@ -245,9 +246,19 @@ namespace CrossCompiler
 			bFound = true;
 			Type = AST::FSemanticSpecifier::ESpecType::PackOffset;
 		}
-		else if ((EDFFlags & EDF_SEMANTIC) == EDF_SEMANTIC)
+		else if ((EDFFlags & EDF_SEMANTIC) == EDF_SEMANTIC || ((EDFFlags & EDF_BITPACK) == EDF_BITPACK))
 		{
-			if (!Scanner.MatchToken(EHlslToken::Identifier))
+			bool bMatched = false;
+			if (Scanner.MatchToken(EHlslToken::Identifier))
+			{
+				bMatched = true;
+			}
+			else if (((EDFFlags & EDF_BITPACK) == EDF_BITPACK) && Scanner.MatchToken(EHlslToken::Literal))
+			{
+				bMatched = true;
+			}
+
+			if (!bMatched)
 			{
 				Scanner.SourceError(TEXT("Expected identifier for semantic!"));
 				return ParseResultError();
@@ -794,11 +805,11 @@ namespace CrossCompiler
 					}
 
 					bool bSemanticFound = false;
-					if ((DeclarationFlags & (EDF_REGISTER | EDF_PACKOFFSET | EDF_SEMANTIC)) != 0)
+					if ((DeclarationFlags & (EDF_REGISTER | EDF_PACKOFFSET | EDF_SEMANTIC | EDF_BITPACK)) != 0)
 					{
 						if (Parser.Scanner.MatchToken(EHlslToken::Colon))
 						{
-							if (ParseColonSpecifier(Parser.Scanner, Allocator, (DeclarationFlags & (EDF_SEMANTIC | EDF_REGISTER | EDF_PACKOFFSET)), &Declaration->Semantic) != EParseResult::Matched)
+							if (ParseColonSpecifier(Parser.Scanner, Allocator, (DeclarationFlags & (EDF_SEMANTIC | EDF_REGISTER | EDF_PACKOFFSET | EDF_BITPACK)), &Declaration->Semantic) != EParseResult::Matched)
 							{
 								Parser.Scanner.SourceError(TEXT("Expected semantic, register or packoffset!\n"));
 								return ParseResultError();
@@ -926,7 +937,7 @@ namespace CrossCompiler
 			// Greedily eat a member declaration, and backtrack if it's a member function
 			const uint32 OriginalToken = Parser.Scanner.GetCurrentTokenIndex();
 			AST::FDeclaratorList* Declaration = nullptr;
-			EParseResult Result = ParseGeneralDeclarationNoSemicolon(Parser, SymbolScope, 0, EDF_CONST_ROW_MAJOR | EDF_SEMANTIC | EDF_MULTIPLE | EDF_TEXTURE_SAMPLER_OR_BUFFER | EDF_INTERPOLATION, Allocator, &Declaration);
+			EParseResult Result = ParseGeneralDeclarationNoSemicolon(Parser, SymbolScope, 0, EDF_CONST_ROW_MAJOR | EDF_SEMANTIC | EDF_MULTIPLE | EDF_TEXTURE_SAMPLER_OR_BUFFER | EDF_INTERPOLATION | EDF_BITPACK, Allocator, &Declaration);
 			bool bTryMemberFunction = false;
 			if (Result == EParseResult::Error)
 			{
