Fix StaticAnalysis warning

diff --git a/Engine/Source/Developer/HlslParser/Private/HlslExpressionParser.inl b/Engine/Source/Developer/HlslParser/Private/HlslExpressionParser.inl
--- a/Engine/Source/Developer/HlslParser/Private/HlslExpressionParser.inl
+++ b/Engine/Source/Developer/HlslParser/Private/HlslExpressionParser.inl
@@ -1015,7 +1015,7 @@ namespace CrossCompiler
 				break;
 			case EHlslToken::Lower:
 			{
-				if (SymbolScope->ScopeType == FSymbolScope::EType::Statement)
+				if (SymbolScope && SymbolScope->ScopeType == FSymbolScope::EType::Statement)
 				{
 					const FHlslToken* Token1 = Scanner.PeekToken(1);
 					
