[HlslParser] Fix null pointer when parsing case statements.

#test Modified DepthOnlyPixelShader.usf to include a switch-case statement with double semicolon `;;` to confirm the issue and fix.
#rnx
#Synced-CL 51745034

diff --git a/Engine/Source/Developer/HlslParser/Private/HlslParser.cpp b/Engine/Source/Developer/HlslParser/Private/HlslParser.cpp
--- a/Engine/Source/Developer/HlslParser/Private/HlslParser.cpp
+++ b/Engine/Source/Developer/HlslParser/Private/HlslParser.cpp
@@ -2109,7 +2109,8 @@ Done:
 						Parser.Scanner.SourceError(TEXT("Internal Error parsing statment inside case list"));
 						return ParseResultError();
 					}
-					else
+					// Ignore null statements, e.g. from a double semicolon `;;`
+					else if (Statement != nullptr)
 					{
 						CaseStatement->Statements.Add(Statement);
 					}
